don't check(false) on failure in geometryprocessing commandutils, as it is confusing to print a stack trace when the error is just a missing parameter

diff --git a/Engine/Source/Programs/GeometryProcessing/Private/Utils/CommandUtils.h b/Engine/Source/Programs/GeometryProcessing/Private/Utils/CommandUtils.h
--- a/Engine/Source/Programs/GeometryProcessing/Private/Utils/CommandUtils.h
+++ b/Engine/Source/Programs/GeometryProcessing/Private/Utils/CommandUtils.h
@@ -14,7 +14,6 @@ namespace UE::CommandUtils
 	// For errors that we can't continue through -- immediately exits the program 
 	inline void Fail()
 	{
-		check(false);
 		FPlatformMisc::RequestExit(true);
 	}
 
