PR #13556: Fix unnecessary namespace usage

diff --git a/Engine/Source/Programs/GeometryProcessing/Private/Utils/CommandUtils.cpp b/Engine/Source/Programs/GeometryProcessing/Private/Utils/CommandUtils.cpp
--- a/Engine/Source/Programs/GeometryProcessing/Private/Utils/CommandUtils.cpp
+++ b/Engine/Source/Programs/GeometryProcessing/Private/Utils/CommandUtils.cpp
@@ -11,7 +11,7 @@ namespace UE::CommandUtils
 		using namespace UE::Geometry;
 
 		FString Path = RequireParam<FString>(InputArg);
-		UE::Geometry::FDynamicMesh3 InputMesh;
+		FDynamicMesh3 InputMesh;
 		ELoadOBJStatus Status = LoadOBJ(TCHAR_TO_ANSI(*Path), InputMesh, Settings);
 		if (Status != ELoadOBJStatus::Success)
 		{
